rename PakFileCompressionFormats to PackageCompressionFormat so that existing DefaultGame.ini files with "PakFileCompressionFormats=" don't override the new BaseGame.ini setting of Oodle

diff --git a/Engine/Config/Android/BaseAndroidGame.ini b/Engine/Config/Android/BaseAndroidGame.ini
--- a/Engine/Config/Android/BaseAndroidGame.ini
+++ b/Engine/Config/Android/BaseAndroidGame.ini
@@ -1,3 +1,3 @@
 [/Script/UnrealEd.ProjectPackagingSettings]
-PakFileCompressionMethod=Mermaid
+PackageCompressionMethod=Mermaid
 
diff --git a/Engine/Config/BaseGame.ini b/Engine/Config/BaseGame.ini
--- a/Engine/Config/BaseGame.ini
+++ b/Engine/Config/BaseGame.ini
@@ -87,12 +87,12 @@ FullRebuild=False
 UsePakFile=True
 bUseIoStore=True
 bCompressed=True
-PakFileCompressionFormats=Oodle
-PakFileAdditionalCompressionOptions=-compressionblocksize=256KB -asynccompression
-PakFileCompressionMethod=Kraken
-PakFileCompressionLevel_Distribution=7
-PakFileCompressionLevel_TestShipping=5
-PakFileCompressionLevel_DebugDevelopment=3
+PackageCompressionFormat=Oodle
+PackageAdditionalCompressionOptions=-compressionblocksize=256KB -asynccompression
+PackageCompressionMethod=Kraken
+PackageCompressionLevel_Distribution=7
+PackageCompressionLevel_TestShipping=5
+PackageCompressionLevel_DebugDevelopment=3
 bForceUseProjectCompressionFormatIgnoreHardwareOverride=False
 bGenerateChunks=False
 bChunkHardReferencesOnly=False
diff --git a/Engine/Config/IOS/BaseIOSGame.ini b/Engine/Config/IOS/BaseIOSGame.ini
--- a/Engine/Config/IOS/BaseIOSGame.ini
+++ b/Engine/Config/IOS/BaseIOSGame.ini
@@ -1,3 +1,3 @@
 [/Script/UnrealEd.ProjectPackagingSettings]
-PakFileCompressionMethod=Mermaid
+PackageCompressionMethod=Mermaid
 
diff --git a/Engine/Plugins/Compression/OodleData/Source/Private/OodleDataPlugin.cpp b/Engine/Plugins/Compression/OodleData/Source/Private/OodleDataPlugin.cpp
--- a/Engine/Plugins/Compression/OodleData/Source/Private/OodleDataPlugin.cpp
+++ b/Engine/Plugins/Compression/OodleData/Source/Private/OodleDataPlugin.cpp
@@ -31,28 +31,28 @@
 [/Script/UnrealEd.ProjectPackagingSettings]
 bCompressed=True
 bForceUseProjectCompressionFormat=False
-PakFileCompressionFormats=Oodle
-PakFileAdditionalCompressionOptions=-compressionblocksize=1MB -asynccompression
-PakFileCompressionMethod=Mermaid
-PakFileCompressionLevel_Distribution=8
-PakFileCompressionLevel_TestShipping=5
-PakFileCompressionLevel_DebugDevelopment=3
+PackageCompressionFormat=Oodle
+PackageAdditionalCompressionOptions=-compressionblocksize=1MB -asynccompression
+PackageCompressionMethod=Mermaid
+PackageCompressionLevel_Distribution=8
+PackageCompressionLevel_TestShipping=5
+PackageCompressionLevel_DebugDevelopment=3
 
 * This can be set in DefaultGame.ini then overrides set up per-platform.
 * 
 * The Engine also has a veto compressionformat set up in the DataDrivenPlatformInfo.ini for each platform in the field
 * "HardwareCompressionFormat"
 * eg. platforms that don't want any software compressor can set "HardwareCompressionFormat=None" and this will override what you
-* set in "PakFileCompressionFormats".
+* set in "PackageCompressionFormat".
 * 
-* The idea is in typical use, you set "PakFileCompressionFormats" for your Game, and you get that compressor on most platforms, but on
+* The idea is in typical use, you set "PackageCompressionFormat" for your Game, and you get that compressor on most platforms, but on
 * some platforms that don't want compression, it automatically turns off.
 * 
 * If you want to force use of your Game.ini compressor (ignore the HardwareCompressionFormat) you can set bForceUseProjectCompressionFormat
 * in ProjectPackagingSettings.
 * 
 * 
-* When using Oodle we recommend "-compressionblocksize=1MB -asynccompression" which can be set with PakFileAdditionalCompressionOptions.
+* When using Oodle we recommend "-compressionblocksize=1MB -asynccompression" which can be set with PackageAdditionalCompressionOptions.
 * 
 * ***************************/
 
diff --git a/Engine/Source/Developer/DeveloperToolSettings/Classes/Settings/ProjectPackagingSettings.h b/Engine/Source/Developer/DeveloperToolSettings/Classes/Settings/ProjectPackagingSettings.h
--- a/Engine/Source/Developer/DeveloperToolSettings/Classes/Settings/ProjectPackagingSettings.h
+++ b/Engine/Source/Developer/DeveloperToolSettings/Classes/Settings/ProjectPackagingSettings.h
@@ -280,41 +280,42 @@ public:
 	/**
 	 * A comma separated list of formats to use for .pak file and IoStore compression. If more than one is specified, the list is in order of priority, with fallbacks to other formats
 	 * in case of errors or unavailability of the format (plugin not enabled, etc).
-	 * Commonly PakFileCompressionFormats=Oodle or PakFileCompressionFormats=None
+	 * Commonly PackageCompressionFormat=Oodle or PackageCompressionFormat=None
 	 */
-	UPROPERTY(config, EditAnywhere, Category = Packaging, AdvancedDisplay, meta = (DisplayName = "Pak File / IoStore Compression Format(s)"))
-	FString PakFileCompressionFormats;
+	UPROPERTY(config, EditAnywhere, Category = Packaging, AdvancedDisplay, meta = (DisplayName = "Package Compression Format"))
+	FString PackageCompressionFormat;
 	
 	/**
-	 * Force use of PakFileCompressionFormats (do not use override HardwareCompressionFormat from DDPI)
+	 * Force use of PackageCompressionFormat (do not use override HardwareCompressionFormat from DDPI)
 	 */
 	UPROPERTY(config, EditAnywhere, Category=Packaging, AdvancedDisplay, meta=(DisplayName="Use this Compression Format not hardware override"))
 	bool bForceUseProjectCompressionFormatIgnoreHardwareOverride;
 
 	/**
 	 * A generic setting for allowing a project to control compression settings during .pak file and iostore compression.
-	 * For instance PakFileAdditionalCompressionOptions=-compressionblocksize=1MB -asynccompression
+	 * For instance PackageAdditionalCompressionOptions=-compressionblocksize=1MB -asynccompression
 	 */
-	UPROPERTY(config, EditAnywhere, Category = Packaging, AdvancedDisplay, meta = (DisplayName = "Pak File / IoStore Compression Commandline Options"))
-	FString PakFileAdditionalCompressionOptions;
+	UPROPERTY(config, EditAnywhere, Category = Packaging, AdvancedDisplay, meta = (DisplayName = "Package Compression Commandline Options"))
+	FString PackageAdditionalCompressionOptions;
 	
 	/**
 	 * For compressors with multiple methods, select one.  eg. for Oodle you may use one of {Kraken,Mermaid,Selkie,Leviathan}
 	 */
-	UPROPERTY(config, EditAnywhere, Category = Packaging, AdvancedDisplay, meta = (DisplayName = "Pak File Compression Method"))
-	FString PakFileCompressionMethod;
+	UPROPERTY(config, EditAnywhere, Category = Packaging, AdvancedDisplay, meta = (DisplayName = "Package Compression Method"))
+	FString PackageCompressionMethod;
 		
 	/*
 	 * For compressors with variable levels, select the encoder effort level, which makes packages smaller but takes more time to encode.
+	 * This does not affect decode speed.  For faster iteration, use lower effort levels (eg. 1)
 	 */
 	UPROPERTY(config, EditAnywhere, Category = Packaging, AdvancedDisplay, meta = (DisplayName = "Encoder Effort Level for Debug & Development"))
-	int32 PakFileCompressionLevel_DebugDevelopment;
+	int32 PackageCompressionLevel_DebugDevelopment;
 	
 	UPROPERTY(config, EditAnywhere, Category = Packaging, AdvancedDisplay, meta = (DisplayName = "Encoder Effort Level for Test & Shipping"))
-	int32 PakFileCompressionLevel_TestShipping;
+	int32 PackageCompressionLevel_TestShipping;
 	
 	UPROPERTY(config, EditAnywhere, Category = Packaging, AdvancedDisplay, meta = (DisplayName = "Encoder Effort Level for Distribution"))
-	int32 PakFileCompressionLevel_Distribution;
+	int32 PackageCompressionLevel_Distribution;
 
 	/** 
 	 * Version name for HTTP Chunk Install Data.
diff --git a/Engine/Source/Programs/AutomationTool/Scripts/CopyBuildToStagingDirectory.Automation.cs b/Engine/Source/Programs/AutomationTool/Scripts/CopyBuildToStagingDirectory.Automation.cs
--- a/Engine/Source/Programs/AutomationTool/Scripts/CopyBuildToStagingDirectory.Automation.cs
+++ b/Engine/Source/Programs/AutomationTool/Scripts/CopyBuildToStagingDirectory.Automation.cs
@@ -2233,7 +2233,7 @@ public partial class Project : CommandUtils
 		}
 		else
 		{
-			PlatformGameConfig.GetString("/Script/UnrealEd.ProjectPackagingSettings", "PakFileCompressionFormats", out CompressionFormats);
+			PlatformGameConfig.GetString("/Script/UnrealEd.ProjectPackagingSettings", "PackageCompressionFormat", out CompressionFormats);
 			// empty CompressionFormats string means "use default" eg. zlib
 		}
 
@@ -2285,12 +2285,12 @@ public partial class Project : CommandUtils
 			// the game may want to control compression settings, but since it may be in a plugin that checks the commandline for the settings, we need to pass
 			// the settings directly on the UnrealPak commandline, and not put it into the batch file lines (plugins can't get the unrealpak command list, and
 			// there's not a great way to communicate random strings down into the plugins during plugin init time)
-			PlatformGameConfig.GetString("/Script/UnrealEd.ProjectPackagingSettings", "PakFileAdditionalCompressionOptions", out AdditionalCompressionOptionsOnCommandLine);
+			PlatformGameConfig.GetString("/Script/UnrealEd.ProjectPackagingSettings", "PackageAdditionalCompressionOptions", out AdditionalCompressionOptionsOnCommandLine);
 
-			//LogInformation("PakFileAdditionalCompressionOptions = {0}", AdditionalCompressionOptionsOnCommandLine);
+			//LogInformation("PackageAdditionalCompressionOptions = {0}", AdditionalCompressionOptionsOnCommandLine);
 
 			string CompressionMethod;
-			PlatformGameConfig.GetString("/Script/UnrealEd.ProjectPackagingSettings", "PakFileCompressionMethod", out CompressionMethod);
+			PlatformGameConfig.GetString("/Script/UnrealEd.ProjectPackagingSettings", "PackageCompressionMethod", out CompressionMethod);
 			if (!string.IsNullOrWhiteSpace(CompressionMethod))
 			{
 				//LogInformation("CompressionMethod = {0}", CompressionMethod);
@@ -2302,15 +2302,15 @@ public partial class Project : CommandUtils
 			// GetInt32 fills out with = 0 if not found
 			if (Params.Distribution)
 			{
-				PlatformGameConfig.GetInt32("/Script/UnrealEd.ProjectPackagingSettings", "PakFileCompressionLevel_Distribution", out CompressionLevel);
+				PlatformGameConfig.GetInt32("/Script/UnrealEd.ProjectPackagingSettings", "PackageCompressionLevel_Distribution", out CompressionLevel);
 			}
 			else if (SC.StageTargetConfigurations.Any(Config => Config == UnrealTargetConfiguration.Test || Config == UnrealTargetConfiguration.Shipping))
 			{
-				PlatformGameConfig.GetInt32("/Script/UnrealEd.ProjectPackagingSettings", "PakFileCompressionLevel_TestShipping", out CompressionLevel);
+				PlatformGameConfig.GetInt32("/Script/UnrealEd.ProjectPackagingSettings", "PackageCompressionLevel_TestShipping", out CompressionLevel);
 			}
 			else
 			{
-				PlatformGameConfig.GetInt32("/Script/UnrealEd.ProjectPackagingSettings", "PakFileCompressionLevel_DebugDevelopment", out CompressionLevel);
+				PlatformGameConfig.GetInt32("/Script/UnrealEd.ProjectPackagingSettings", "PackageCompressionLevel_DebugDevelopment", out CompressionLevel);
 			}
 
 			if (CompressionLevel != 0)
