Fix packaging issue with config enum saved in an invalid state

diff --git a/Engine/Source/Developer/DeveloperToolSettings/Classes/Settings/ProjectPackagingSettings.h b/Engine/Source/Developer/DeveloperToolSettings/Classes/Settings/ProjectPackagingSettings.h
--- a/Engine/Source/Developer/DeveloperToolSettings/Classes/Settings/ProjectPackagingSettings.h
+++ b/Engine/Source/Developer/DeveloperToolSettings/Classes/Settings/ProjectPackagingSettings.h
@@ -31,7 +31,7 @@ enum class EProjectPackagingBuildConfigurations : uint8
 	PPBC_Shipping UMETA(DisplayName="Shipping"),
 
 	/** Number of entries in the enum. */
-	PPBC_MAX
+	PPBC_MAX UMETA(Hidden)
 };
 
 /**
