Fix a typo in function comment for make binary config setting in package settings.

diff --git a/Engine/Source/Developer/DeveloperToolSettings/Classes/Settings/ProjectPackagingSettings.h b/Engine/Source/Developer/DeveloperToolSettings/Classes/Settings/ProjectPackagingSettings.h
--- a/Engine/Source/Developer/DeveloperToolSettings/Classes/Settings/ProjectPackagingSettings.h
+++ b/Engine/Source/Developer/DeveloperToolSettings/Classes/Settings/ProjectPackagingSettings.h
@@ -243,7 +243,7 @@ public:
 	UPROPERTY(config, EditAnywhere, Category = Packaging)
 	bool bUseIoStore;
 
-	/** If enabled, staging will make a binary config fie for faster loading. */
+	/** If enabled, staging will make a binary config file for faster loading. */
 	UPROPERTY(config, EditAnywhere, Category = Packaging)
 	bool bMakeBinaryConfig;
 
