- Added a new ProjectPackaging setting - LaunchOnTarget which is now shown int he Platforms menu when multiple Game targets exist (Lyra).
- Added support for passing the LaunchOnTarget along to the LaunchOn process

diff --git a/Engine/Source/Developer/DeveloperToolSettings/Classes/Settings/ProjectPackagingSettings.h b/Engine/Source/Developer/DeveloperToolSettings/Classes/Settings/ProjectPackagingSettings.h
--- a/Engine/Source/Developer/DeveloperToolSettings/Classes/Settings/ProjectPackagingSettings.h
+++ b/Engine/Source/Developer/DeveloperToolSettings/Classes/Settings/ProjectPackagingSettings.h
@@ -191,6 +191,10 @@ public:
 	UPROPERTY(config, EditAnywhere, Category=Project)
 	FString BuildTarget;
 
+	/** Name of the target to use for LaunchOn (only Game/Client targets) */
+	UPROPERTY(config)
+	FString LaunchOnTarget;
+
 	/** The directory to which the packaged project will be copied. */
 	UPROPERTY(config, EditAnywhere, Category=Project)
 	FDirectoryPath StagingDirectory;
@@ -656,6 +660,9 @@ public:
 	/** Gets the current build target, checking that it's valid, and the default build target if it is not */
 	const FTargetInfo* GetBuildTargetInfo() const;
 
+	/** Gets the current launch ojn target, checking that it's valid, and the default build target if it is not */
+	const FTargetInfo* GetLaunchOnTargetInfo() const;
+
 	/** For non-default object instances, this will LoadConfig for a specific platform,*/
 	void LoadSettingsForPlatform(FString PlatformName)
 	{
diff --git a/Engine/Source/Developer/DeveloperToolSettings/Private/ProjectPackagingSettings.cpp b/Engine/Source/Developer/DeveloperToolSettings/Private/ProjectPackagingSettings.cpp
--- a/Engine/Source/Developer/DeveloperToolSettings/Private/ProjectPackagingSettings.cpp
+++ b/Engine/Source/Developer/DeveloperToolSettings/Private/ProjectPackagingSettings.cpp
@@ -193,43 +193,29 @@ TArray<EProjectPackagingBuildConfigurations> UProjectPackagingSettings::GetValid
 	return Configurations;
 }
 
-const FTargetInfo* UProjectPackagingSettings::GetBuildTargetInfo() const
+static const FTargetInfo* FindBestTargetInfo(const FString& TargetName, bool bContentOnlyUsesEngineTargets, bool* bOutIsProjectTarget=nullptr)
 {
-	const FTargetInfo* DefaultGameTarget = nullptr;
-	const FTargetInfo* DefaultClientTarget = nullptr;
-	for (const FTargetInfo& Target : FDesktopPlatformModule::Get()->GetTargetsForCurrentProject())
+	bool bUseEngineTargets = false;
+	if (bContentOnlyUsesEngineTargets)
 	{
-		if (Target.Name == BuildTarget)
-		{
-			return &Target;
-		}
-		else if (Target.Type == EBuildTargetType::Game && (DefaultGameTarget == nullptr || Target.Name < DefaultGameTarget->Name))
-		{
-			DefaultGameTarget = &Target;
-		}
-		else if (Target.Type == EBuildTargetType::Client && (DefaultClientTarget == nullptr || Target.Name < DefaultClientTarget->Name))
-		{
-			DefaultClientTarget = &Target;
-		}
+		// Collect build targets. Content-only projects use Engine targets.
+		FProjectStatus ProjectStatus;
+		bUseEngineTargets =!(IProjectManager::Get().QueryStatusForCurrentProject(ProjectStatus) && ProjectStatus.bCodeBasedProject);
+	}
+
+	if (bOutIsProjectTarget != nullptr)
+	{
+		*bOutIsProjectTarget = !bUseEngineTargets;
 	}
-	return (DefaultGameTarget != nullptr) ? DefaultGameTarget : DefaultClientTarget;
-}
 
-const FTargetInfo* UProjectPackagingSettings::GetBuildTargetInfoForPlatform(FName PlatformName, bool& bOutIsProjectTarget) const
-{
-	// Collect build targets. Content-only projects use Engine targets.
-	FProjectStatus ProjectStatus;
-	bool bHasCode = IProjectManager::Get().QueryStatusForCurrentProject(ProjectStatus) && ProjectStatus.bCodeBasedProject;
-	bOutIsProjectTarget = bHasCode;
-	
 	IDesktopPlatform* DesktopPlatform = FDesktopPlatformModule::Get();
-	const TArray<FTargetInfo>& TargetsRef = bHasCode ? DesktopPlatform->GetTargetsForCurrentProject() : DesktopPlatform->GetTargetsForProject(FString());
+	const TArray<FTargetInfo>& TargetsRef = bUseEngineTargets ? DesktopPlatform->GetTargetsForProject(FString()) : DesktopPlatform->GetTargetsForCurrentProject();
 
 	const FTargetInfo* DefaultGameTarget = nullptr;
 	const FTargetInfo* DefaultClientTarget = nullptr;
 	for (const FTargetInfo& Target : TargetsRef)
 	{
-		if (Target.Name == GetBuildTargetForPlatform(PlatformName))
+		if (Target.Name == TargetName)
 		{
 			return &Target;
 		}
@@ -243,8 +229,25 @@ const FTargetInfo* UProjectPackagingSettings::GetBuildTargetInfoForPlatform(FNam
 		}
 	}
 	return (DefaultGameTarget != nullptr) ? DefaultGameTarget : DefaultClientTarget;
+
+}
+
+const FTargetInfo* UProjectPackagingSettings::GetBuildTargetInfo() const
+{
+	return FindBestTargetInfo(BuildTarget, false);
+}
+
+const FTargetInfo* UProjectPackagingSettings::GetBuildTargetInfoForPlatform(FName PlatformName, bool& bOutIsProjectTarget) const
+{
+	return FindBestTargetInfo(GetBuildTargetForPlatform(PlatformName), true, &bOutIsProjectTarget);
 }
 
+const FTargetInfo* UProjectPackagingSettings::GetLaunchOnTargetInfo() const
+{
+	return FindBestTargetInfo(LaunchOnTarget, true);
+}
+
+
 EProjectPackagingBuildConfigurations UProjectPackagingSettings::GetBuildConfigurationForPlatform(FName PlatformName) const
 {
 	const EProjectPackagingBuildConfigurations* Value = PerPlatformBuildConfig.Find(PlatformName);
diff --git a/Engine/Source/Developer/LauncherServices/Private/Launcher/LauncherWorker.cpp b/Engine/Source/Developer/LauncherServices/Private/Launcher/LauncherWorker.cpp
--- a/Engine/Source/Developer/LauncherServices/Private/Launcher/LauncherWorker.cpp
+++ b/Engine/Source/Developer/LauncherServices/Private/Launcher/LauncherWorker.cpp
@@ -426,6 +426,11 @@ FString FLauncherWorker::CreateUATCommand( const ILauncherProfileRef& InProfile,
 		UATCommand += (TEXT(" -cookflavor=") + Join(OptionalCookFlavors, TEXT("+")));
 	}
 
+	if (InProfile->GetBuildTarget().Len() > 0)
+	{
+		UATCommand += TEXT(" -target=") + InProfile->GetBuildTarget();
+	}
+
 	// device list
 	FString DeviceNames = TEXT("");
 	FString DeviceCommand = TEXT("");
diff --git a/Engine/Source/Developer/LauncherServices/Private/Profiles/LauncherProfile.h b/Engine/Source/Developer/LauncherServices/Private/Profiles/LauncherProfile.h
--- a/Engine/Source/Developer/LauncherServices/Private/Profiles/LauncherProfile.h
+++ b/Engine/Source/Developer/LauncherServices/Private/Profiles/LauncherProfile.h
@@ -395,6 +395,11 @@ public:
 		return BuildConfiguration;
 	}
 
+	virtual const FString& GetBuildTarget() const override
+	{
+		return BuildTargetName;
+	}
+
 	virtual EBuildConfiguration GetCookConfiguration( ) const override
 	{
 		return CookConfiguration;
@@ -2017,6 +2022,16 @@ public:
 		}
 	}
 
+	virtual void SetBuildTarget( const FString& TargetName ) override
+	{
+		if (BuildTargetName != TargetName)
+		{
+			BuildTargetName = TargetName;
+
+			Validate();
+		}
+	}
+
 	virtual void SetCookConfiguration( EBuildConfiguration Configuration ) override
 	{
 		if (CookConfiguration != Configuration)
@@ -2759,6 +2774,9 @@ private:
 	// Holds the desired build configuration (only used if creating new builds).
 	EBuildConfiguration BuildConfiguration;
 
+	// Holds the name of the target (matching a .cs file) to build. Needed when multiple targets of a type exist
+	FString BuildTargetName;
+
 	// Holds the build mode.
 	// Holds the build configuration name of the cooker.
 	EBuildConfiguration CookConfiguration;
diff --git a/Engine/Source/Developer/LauncherServices/Public/ILauncherProfile.h b/Engine/Source/Developer/LauncherServices/Public/ILauncherProfile.h
--- a/Engine/Source/Developer/LauncherServices/Public/ILauncherProfile.h
+++ b/Engine/Source/Developer/LauncherServices/Public/ILauncherProfile.h
@@ -512,6 +512,14 @@ public:
 	 */
 	virtual EBuildConfiguration GetBuildConfiguration( ) const = 0;
 
+	/**
+	 * Gets the build target.
+	 *
+	 * @Return Target name to build/run (it would match a .Target.cs file)
+	 * @see SetBuildTarget
+	 */
+	virtual const FString& GetBuildTarget() const = 0;
+
 	/**
 	 * Gets the build configuration name of the cooker.
 	 *
@@ -1016,11 +1024,18 @@ public:
 	 */
 	virtual void SetBuildConfiguration( EBuildConfiguration Configuration ) = 0;
 
+	/**
+	 * Sets the build target.
+	 *
+	 * @param TargetName The target name to set (it would match a .Target.cs file)
+	 * @see GetBuildTarget
+	 */
+	virtual void SetBuildTarget( const FString& TargetName ) = 0;
 	/**
 	 * Sets the build configuration of the cooker.
 	 *
 	 * @param Configuration The cooker's build configuration to set.
-	 * @see GetBuildConfigurationName
+	 * @see GetCookConfiguration
 	 */
 	virtual void SetCookConfiguration( EBuildConfiguration Configuration ) = 0;
 
diff --git a/Engine/Source/Editor/TurnkeySupport/Private/TurnkeySupportModule.cpp b/Engine/Source/Editor/TurnkeySupport/Private/TurnkeySupportModule.cpp
--- a/Engine/Source/Editor/TurnkeySupport/Private/TurnkeySupportModule.cpp
+++ b/Engine/Source/Editor/TurnkeySupport/Private/TurnkeySupportModule.cpp
@@ -600,12 +600,24 @@ public:
 		PackagingSettings->SetBuildTargetForPlatform(Info->IniPlatformName, TargetName);
 		PackagingSettings->SaveConfig();
 	}
-	
+
 	static bool PackageBuildTargetIsChecked(const PlatformInfo::FTargetPlatformInfo* Info, FString TargetName)
 	{
 		return GetDefault<UProjectPackagingSettings>()->GetBuildTargetForPlatform(Info->IniPlatformName) == TargetName;
 	}
 
+	static void SetLaunchOnBuildTarget(FString TargetName)
+	{
+		UProjectPackagingSettings* PackagingSettings = GetMutableDefault<UProjectPackagingSettings>();
+		PackagingSettings->LaunchOnTarget = TargetName;
+		PackagingSettings->SaveConfig();
+	}
+
+	static bool LaunchOnBuildTargetIsChecked(FString TargetName)
+	{
+		return GetDefault<UProjectPackagingSettings>()->GetLaunchOnTargetInfo()->Name == TargetName;
+	}
+
 	static void SetCookOnTheFly()
 	{
 		UCookerSettings* CookerSettings = GetMutableDefault<UCookerSettings>();
@@ -1529,7 +1541,6 @@ void FTurnkeySupportModule::MakeQuickLaunchItems(class UToolMenu* Menu, FOnQuick
 			);
 		}
 
-
 		// now kick-off any devices that need to be updated
 		if (DeviceIdsToQuery.Num() > 0)
 		{
@@ -1537,6 +1548,39 @@ void FTurnkeySupportModule::MakeQuickLaunchItems(class UToolMenu* Menu, FOnQuick
 		}
 	}
 	));
+
+	FProjectStatus ProjectStatus;
+	bool bHasCode = IProjectManager::Get().QueryStatusForCurrentProject(ProjectStatus) && ProjectStatus.bCodeBasedProject;
+
+	IDesktopPlatform* DesktopPlatform = FDesktopPlatformModule::Get();
+	TArray<FTargetInfo> Targets = bHasCode ? DesktopPlatform->GetTargetsForCurrentProject() : DesktopPlatform->GetTargetsForProject(FString());
+	const TArray<FTargetInfo> ValidTargets = Targets.FilterByPredicate([](const FTargetInfo& Target)
+		{
+			return Target.Type == EBuildTargetType::Game;// || Target.Type == EBuildTargetType::Client || Target.Type == EBuildTargetType::Server;
+		});
+
+	if (ValidTargets.Num() > 1)
+	{
+		FToolMenuSection& LaunchTargetSection = Menu->AddSection("QuickLaunchTarget", LOCTEXT("QuickLaunchTarget", "Quick Launch Game Target"));
+
+		for (const FTargetInfo& Target : ValidTargets)
+		{
+			LaunchTargetSection.AddMenuEntry(
+				NAME_None,
+				FText::FromString(Target.Name),
+				FText::Format(LOCTEXT("PackageTargetName", "Package the '{0}' target."), FText::FromString(Target.Name)),
+				FSlateIcon(),
+				FUIAction(
+					FExecuteAction::CreateStatic(&FTurnkeySupportCallbacks::SetLaunchOnBuildTarget, Target.Name),
+					FCanExecuteAction(),
+					FIsActionChecked::CreateStatic(&FTurnkeySupportCallbacks::LaunchOnBuildTargetIsChecked, Target.Name)
+				),
+				EUserInterfaceActionType::RadioButton
+			);
+		}
+
+	}
+
 }
 
 TSharedRef<SWidget> FTurnkeySupportModule::MakeTurnkeyMenuWidget() const
diff --git a/Engine/Source/Editor/UnrealEd/Private/PlayLevelViaLauncher.cpp b/Engine/Source/Editor/UnrealEd/Private/PlayLevelViaLauncher.cpp
--- a/Engine/Source/Editor/UnrealEd/Private/PlayLevelViaLauncher.cpp
+++ b/Engine/Source/Editor/UnrealEd/Private/PlayLevelViaLauncher.cpp
@@ -266,6 +266,7 @@ void UEditorEngine::StartPlayUsingLauncherSession(FRequestPlaySessionParams& InR
 		}
 	}
 
+	LauncherProfile->SetBuildTarget(GetDefault<UProjectPackagingSettings>()->GetLaunchOnTargetInfo()->Name);
 	LauncherProfile->SetCookMode(CurrentLauncherCookMode);
 	LauncherProfile->SetUnversionedCooking(!bIncrimentalCooking);
 	LauncherProfile->SetIncrementalCooking(bIncrimentalCooking);
