Include debug files in non-shipping builds by default
 - Debug files are now included in all non-shipping builds, and optionally can be included in shipping builds if the IncludeDebugFiles flag is enabled in the project packaging settings.

diff --git a/Engine/Source/Developer/DeveloperToolSettings/Classes/Settings/ProjectPackagingSettings.h b/Engine/Source/Developer/DeveloperToolSettings/Classes/Settings/ProjectPackagingSettings.h
--- a/Engine/Source/Developer/DeveloperToolSettings/Classes/Settings/ProjectPackagingSettings.h
+++ b/Engine/Source/Developer/DeveloperToolSettings/Classes/Settings/ProjectPackagingSettings.h
@@ -215,8 +215,8 @@ public:
 	UPROPERTY(config, EditAnywhere, Category=Project)
 	bool ForDistribution;
 
-	/** If enabled, debug files will be included in the packaged game */
-	UPROPERTY(config, EditAnywhere, Category=Project)
+	/** If enabled, debug files will be included in staged shipping builds. */
+	UPROPERTY(config, EditAnywhere, Category=Project, meta = (DisplayName = "Include Debug Files in Shipping Builds"))
 	bool IncludeDebugFiles;
 
 	/** If enabled, then the project's Blueprint assets (including structs and enums) will be intermediately converted into C++ and used in the packaged project (in place of the .uasset files).*/
diff --git a/Engine/Source/Editor/TurnkeySupport/Private/TurnkeySupportModule.cpp b/Engine/Source/Editor/TurnkeySupport/Private/TurnkeySupportModule.cpp
--- a/Engine/Source/Editor/TurnkeySupport/Private/TurnkeySupportModule.cpp
+++ b/Engine/Source/Editor/TurnkeySupport/Private/TurnkeySupportModule.cpp
@@ -441,11 +441,6 @@ public:
 				BuildCookRunParams += TEXT(" -distribution");
 			}
 
-			if (!PackagingSettings->IncludeDebugFiles)
-			{
-				BuildCookRunParams += TEXT(" -nodebuginfo");
-			}
-
 			if (PackagingSettings->bGenerateChunks)
 			{
 				BuildCookRunParams += TEXT(" -manifests");
@@ -491,6 +486,11 @@ public:
 					BuildCookRunParams += FString::Printf(TEXT(" -clientconfig=%s"), LexToString(ConfigurationInfo.Configuration));
 				}
 			}
+
+			if (ConfigurationInfo.Configuration == EBuildConfiguration::Shipping && !PackagingSettings->IncludeDebugFiles)
+			{
+				BuildCookRunParams += TEXT(" -nodebuginfo");
+			}
 		}
 // 		else if (Mode == EPrepareContentMode::PrepareForDebugging)
 // 		{
