Fixed a shelf vs changelist mismatch

diff --git a/Engine/Source/Developer/DeveloperToolSettings/Private/PlatformsMenuSettings.cpp b/Engine/Source/Developer/DeveloperToolSettings/Private/PlatformsMenuSettings.cpp
--- a/Engine/Source/Developer/DeveloperToolSettings/Private/PlatformsMenuSettings.cpp
+++ b/Engine/Source/Developer/DeveloperToolSettings/Private/PlatformsMenuSettings.cpp
@@ -81,7 +81,7 @@ void UPlatformsMenuSettings::SetBuildTargetForPlatform(FName PlatformName, FStri
 
 const FTargetInfo* UPlatformsMenuSettings::GetLaunchOnTargetInfo() const
 {
-	return FindBestTargetInfo(LaunchOnTarget, true);
+	return FindBestTargetInfo(LaunchOnTarget, true, nullptr);
 }
 
 #undef LOCTEXT_NAMESPACE
