Change default ASTC encoding preset from Fast to Medium
Deleted unused DefaultPVRTCQuality

diff --git a/Engine/Config/BaseEngine.ini b/Engine/Config/BaseEngine.ini
--- a/Engine/Config/BaseEngine.ini
+++ b/Engine/Config/BaseEngine.ini
@@ -2956,8 +2956,7 @@ SDKAPILevel=latest
 NDKAPILevel=android-26
 
 [/Script/UnrealEd.CookerSettings]
-DefaultPVRTCQuality=1
-DefaultASTCQualityBySpeed=1
+DefaultASTCQualityBySpeed=2
 DefaultASTCQualityBySize=3
 +ClassesExcludedOnDedicatedServer=WidgetBlueprint
 +ClassesExcludedOnDedicatedServer=GroupActor
diff --git a/Engine/Source/Developer/DeveloperToolSettings/Private/CookerSettings.cpp b/Engine/Source/Developer/DeveloperToolSettings/Private/CookerSettings.cpp
--- a/Engine/Source/Developer/DeveloperToolSettings/Private/CookerSettings.cpp
+++ b/Engine/Source/Developer/DeveloperToolSettings/Private/CookerSettings.cpp
@@ -2,6 +2,7 @@
 
 #include "CookerSettings.h"
 #include "UObject/UnrealType.h"
+#include "HAL/IConsoleManager.h"
 
 #include UE_INLINE_GENERATED_CPP_BY_NAME(CookerSettings)
 
@@ -22,8 +23,13 @@ UCookerSettings::UCookerSettings(const FObjectInitializer& ObjectInitializer)
 	, bCookBlueprintComponentTemplateData(false)
 {
 	SectionName = TEXT("Cooker");
-	DefaultASTCQualityBySize = 3;
-	DefaultASTCQualityBySpeed = 3;
+	
+	DefaultASTCQualityBySpeed = 2; // Medium preset
+	DefaultASTCQualityBySize = 3;  // 6x6
+	{
+		static IConsoleVariable* ASTCTextureCompressorCVar = IConsoleManager::Get().FindConsoleVariable(TEXT("cook.ASTCTextureCompressor"));
+		DefaultASTCCompressor = (ASTCTextureCompressorCVar && ASTCTextureCompressorCVar->GetInt() != 0) ? ETextureFormatASTCCompressor::Arm : ETextureFormatASTCCompressor::IntelISPC;
+	}
 }
 
 void UCookerSettings::PostInitProperties()
