[NO NON-PII COMMIT MESSAGE]

diff --git a/Engine/Source/Developer/DeveloperToolSettings/Classes/Settings/ProjectPackagingSettings.h b/Engine/Source/Developer/DeveloperToolSettings/Classes/Settings/ProjectPackagingSettings.h
--- a/Engine/Source/Developer/DeveloperToolSettings/Classes/Settings/ProjectPackagingSettings.h
+++ b/Engine/Source/Developer/DeveloperToolSettings/Classes/Settings/ProjectPackagingSettings.h
@@ -68,7 +68,7 @@ enum class EProjectPackagingBuild
 	/** Never build. */
 	Never UMETA(DisplayName="Never"),
 
-	/** Default (if the Never build. */
+	/** Default (if the never build.) */
 	IfProjectHasCode UMETA(DisplayName="If project has code, or running a locally built editor"),
 
 	/** If we're not packaging from a promoted build. */
