Project Packaging Settings - Ensure that if bUseIoStore is False and bUseZenStore is True, that bUseZenStore is treated as if it was False, and emit a warning explaining that this combination is unsupported.

diff --git a/Engine/Source/Developer/DeveloperToolSettings/Private/ProjectPackagingSettings.cpp b/Engine/Source/Developer/DeveloperToolSettings/Private/ProjectPackagingSettings.cpp
--- a/Engine/Source/Developer/DeveloperToolSettings/Private/ProjectPackagingSettings.cpp
+++ b/Engine/Source/Developer/DeveloperToolSettings/Private/ProjectPackagingSettings.cpp
@@ -12,6 +12,7 @@
 
 #define LOCTEXT_NAMESPACE "SettingsClasses"
 
+DEFINE_LOG_CATEGORY_STATIC(LogProjectPackagingSettings, Log, Log);
 
 /* UProjectPackagingSettings interface
  *****************************************************************************/
@@ -39,6 +40,15 @@ void UProjectPackagingSettings::PostInitProperties()
 	FixCookingPaths();
 
 	Super::PostInitProperties();
+
+	if (bUseZenStore && !bUseIoStore)
+	{
+		UE_LOG(LogProjectPackagingSettings, Warning, TEXT("bUseZenStore is enabled in your project packaging settings while bUseIoStore is disabled. "
+															"This is not a supported combination.  bUseZenStore is being treated as disabled. "
+															"To prevent this warning, either (preferred) set bUseIoStore=True or set bUseZenStore=False"
+															"in your project packaging settings in your DefaultGame.ini"));
+		bUseZenStore = false;
+	}
 }
 
 void UProjectPackagingSettings::FixCookingPaths()
diff --git a/Engine/Source/Programs/AutomationTool/Turnkey/Commands/CreateBuild.cs b/Engine/Source/Programs/AutomationTool/Turnkey/Commands/CreateBuild.cs
--- a/Engine/Source/Programs/AutomationTool/Turnkey/Commands/CreateBuild.cs
+++ b/Engine/Source/Programs/AutomationTool/Turnkey/Commands/CreateBuild.cs
@@ -274,7 +274,10 @@ namespace Turnkey.Commands
 
 			if (bCook)
 			{
-				AddIfTrue("bUseZenStore", "-zenstore");
+				if (GetBool("bUseIoStore"))
+				{
+					AddIfTrue("bUseZenStore", "-zenstore");
+				}
 				AddIfTrue("bSkipEditorContent", "-SkipCookingEditorContent");
 			}
 			
