Cooker: Created WarningsAsErrorsOnCook toggle to project settings
#11623
#rnx

diff --git a/Engine/Source/Developer/DeveloperToolSettings/Classes/Settings/ProjectPackagingSettings.h b/Engine/Source/Developer/DeveloperToolSettings/Classes/Settings/ProjectPackagingSettings.h
--- a/Engine/Source/Developer/DeveloperToolSettings/Classes/Settings/ProjectPackagingSettings.h
+++ b/Engine/Source/Developer/DeveloperToolSettings/Classes/Settings/ProjectPackagingSettings.h
@@ -486,6 +486,12 @@ public:
 	UPROPERTY(config, EditAnywhere, Category=Packaging, AdvancedDisplay, meta=(DisplayName="Cook only maps (this only affects cookall)"))
 	bool bCookMapsOnly;
 
+	/**
+	 * Cook with Warnings As Errors
+	 */
+	UPROPERTY(config, EditAnywhere, Category=Packaging, AdvancedDisplay, meta=(DisplayName="Cook with Warnings As Errors enabled"))
+	bool bTreatWarningsAsErrorsOnCook;
+
 	/**
 	* Encrypt ini files inside of the pak file
 	* NOTE: Replaced by the settings inside the cryptokeys system. Kept here for legacy migration purposes.
diff --git a/Engine/Source/Editor/UnrealEd/Private/CookOnTheFlyServer.cpp b/Engine/Source/Editor/UnrealEd/Private/CookOnTheFlyServer.cpp
--- a/Engine/Source/Editor/UnrealEd/Private/CookOnTheFlyServer.cpp
+++ b/Engine/Source/Editor/UnrealEd/Private/CookOnTheFlyServer.cpp
@@ -7137,6 +7137,12 @@ void UCookOnTheFlyServer::SetInitializeConfigSettings(UE::Cook::FInitializeConfi
 	CookerIdleWarningSeverity = ParseLogVerbosityFromString(Severity);
 
 	bCookFastStartup = FParse::Param(FCommandLine::Get(), TEXT("cookfaststartup"));
+
+	const UProjectPackagingSettings* PackagingSettings = Cast<UProjectPackagingSettings>(UProjectPackagingSettings::StaticClass()->GetDefaultObject());
+	if (PackagingSettings->bTreatWarningsAsErrorsOnCook)
+	{
+		GWarn->TreatWarningsAsErrors = true;
+	}
 }
 
 void UCookOnTheFlyServer::ParseCookFilters()
