Correct a tooltip.
#rnx

diff --git a/Engine/Source/Developer/DeveloperToolSettings/Classes/Settings/ProjectPackagingSettings.h b/Engine/Source/Developer/DeveloperToolSettings/Classes/Settings/ProjectPackagingSettings.h
--- a/Engine/Source/Developer/DeveloperToolSettings/Classes/Settings/ProjectPackagingSettings.h
+++ b/Engine/Source/Developer/DeveloperToolSettings/Classes/Settings/ProjectPackagingSettings.h
@@ -329,7 +329,7 @@ public:
 	EAssetRegistryWritebackMethod WriteBackMetadataToAssetRegistry;
 
 	/**
-	* Whether or not to write a json summary file that contains size information to the cooked Metadata/PluginJsons directory
+	* Whether or not to write a json summary file that contains size information to the cooked Metadata directory
 	*/
 	UPROPERTY(config, EditAnywhere, Category = Packaging, AdvancedDisplay, meta = (EditCondition = "WriteBackMetadataToAssetRegistry != EAssetRegistryWritebackMethod::Disabled"))
 	bool bWritePluginSizeSummaryJsons;
