PR #11096: Cooking: Rename bIterativeCookingForFileCookContent DisplayName to match UE5 UI layout

#rnx

diff --git a/Engine/Source/Developer/DeveloperToolSettings/Classes/CookerSettings.h b/Engine/Source/Developer/DeveloperToolSettings/Classes/CookerSettings.h
--- a/Engine/Source/Developer/DeveloperToolSettings/Classes/CookerSettings.h
+++ b/Engine/Source/Developer/DeveloperToolSettings/Classes/CookerSettings.h
@@ -116,8 +116,8 @@ public:
 	UPROPERTY(GlobalConfig, EditAnywhere, Category = Cooker, meta = (DisplayName = "Iterative cooking for builds launched from the editor (launch on)"))
 	bool bIterativeCookingForLaunchOn;
 
-	/** Enable -iterate when triggering from File dropdown menu */
-	UPROPERTY(GlobalConfig, EditAnywhere, Category = Cooker, meta = (DisplayName = "Iterative cooking for the File->Cook Content menu item"))
+	/** Enable -iterate when triggering from Platforms dropdown menu */
+	UPROPERTY(GlobalConfig, EditAnywhere, Category = Cooker, meta = (DisplayName = "Iterative cooking for the Platforms->Cook Content menu item"))
 	bool bIterativeCookingForFileCookContent;
 
 	/** Enable -cookonthefly for launch on */
