Display name and tooltip update for bUseZenStore packaging setting.

diff --git a/Engine/Source/Developer/DeveloperToolSettings/Classes/Settings/ProjectPackagingSettings.h b/Engine/Source/Developer/DeveloperToolSettings/Classes/Settings/ProjectPackagingSettings.h
--- a/Engine/Source/Developer/DeveloperToolSettings/Classes/Settings/ProjectPackagingSettings.h
+++ b/Engine/Source/Developer/DeveloperToolSettings/Classes/Settings/ProjectPackagingSettings.h
@@ -267,8 +267,8 @@ public:
 	UPROPERTY(config, EditAnywhere, Category = Packaging)
 	bool bUseIoStore;
 
-	/** If enabled, use Zen storage server for storing and fetching cooked data instead of using the local file system.  */
-	UPROPERTY(config, EditAnywhere, Category = Packaging)
+	/** If enabled, use Zen Server for storing and fetching cooked data instead of using the local file system.  */
+	UPROPERTY(config, EditAnywhere, Category = Packaging, meta=(DisplayName="Use Zen Server as cooked output store"))
 	bool bUseZenStore;
 
 	/** If enabled, staging will make a binary config file for faster loading. */
