- When "If running a locally built editor" is used in Packaging Settings, the code now looks at the InstalledBuild flag, along with the PromotedBuild flag (to match a BuildUAT in LauncherProfile)

diff --git a/Engine/Source/Developer/DeveloperToolSettings/Classes/Settings/ProjectPackagingSettings.h b/Engine/Source/Developer/DeveloperToolSettings/Classes/Settings/ProjectPackagingSettings.h
--- a/Engine/Source/Developer/DeveloperToolSettings/Classes/Settings/ProjectPackagingSettings.h
+++ b/Engine/Source/Developer/DeveloperToolSettings/Classes/Settings/ProjectPackagingSettings.h
@@ -73,7 +73,7 @@ enum class EProjectPackagingBuild
 	/** Default (if the never build.) */
 	IfProjectHasCode UMETA(DisplayName="If project has code, or running a locally built editor"),
 
-	/** If we're not packaging from a promoted build. */
+	/** If we're not packaging from a promoted or installed build. */
 	IfEditorWasBuiltLocally UMETA(DisplayName="If running a locally built editor")
 };
 
diff --git a/Engine/Source/Editor/TurnkeySupport/Private/TurnkeySupportModule.cpp b/Engine/Source/Editor/TurnkeySupport/Private/TurnkeySupportModule.cpp
--- a/Engine/Source/Editor/TurnkeySupport/Private/TurnkeySupportModule.cpp
+++ b/Engine/Source/Editor/TurnkeySupport/Private/TurnkeySupportModule.cpp
@@ -208,7 +208,7 @@ protected:
 		}
 		else if (PackagingSettings->Build == EProjectPackagingBuild::IfEditorWasBuiltLocally)
 		{
-			bBuild = !FApp::GetEngineIsPromotedBuild();
+			bBuild = !FApp::GetEngineIsPromotedBuild() && !FApp::IsEngineInstalled();
 		}
 
 		return bBuild;
