- Fixed how Build setting is obliterated every editor boot

diff --git a/Engine/Config/BaseGame.ini b/Engine/Config/BaseGame.ini
--- a/Engine/Config/BaseGame.ini
+++ b/Engine/Config/BaseGame.ini
@@ -87,6 +87,7 @@ MinWindowHeight=16
 [/Script/UnrealEd.ProjectPackagingSettings]
 @EngineCustomBuilds=Name
 @ProjectCustomBuilds=Name
+Build=IfProjectHasCode
 BuildConfiguration=PPBC_Development
 FullRebuild=False
 UsePakFile=True
diff --git a/Engine/Source/Developer/DeveloperToolSettings/Private/ProjectPackagingSettings.cpp b/Engine/Source/Developer/DeveloperToolSettings/Private/ProjectPackagingSettings.cpp
--- a/Engine/Source/Developer/DeveloperToolSettings/Private/ProjectPackagingSettings.cpp
+++ b/Engine/Source/Developer/DeveloperToolSettings/Private/ProjectPackagingSettings.cpp
@@ -34,9 +34,6 @@ UProjectPackagingSettings::UProjectPackagingSettings(const FObjectInitializer& O
 
 void UProjectPackagingSettings::PostInitProperties()
 {
-	// Build code projects by default
-	Build = EProjectPackagingBuild::IfProjectHasCode;
-
 	FixCookingPaths();
 
 	Super::PostInitProperties();
